ElectraDecoders: Added an MJPEG decoder for Electra

#rnx

diff --git a/Engine/Plugins/Media/MJPEGDecoderElectra/MJPEGDecoderElectra.uplugin b/Engine/Plugins/Media/MJPEGDecoderElectra/MJPEGDecoderElectra.uplugin
new file mode 100644
--- /dev/null
+++ b/Engine/Plugins/Media/MJPEGDecoderElectra/MJPEGDecoderElectra.uplugin
@@ -0,0 +1,30 @@
+{
+	"FileVersion": 3,
+	"Version": 1,
+	"VersionName": "1.0",
+	"FriendlyName": "MotionJPEG Decoder for Electra",
+	"Description": "Implements video playback of MotionJPEG encoded video files",
+	"CreatedBy": "Epic Games, Inc.",
+	"CreatedByURL": "https://epicgames.com",
+	"Category": "Media Players",
+	"EnabledByDefault": false,
+	"CanContainContent": false,
+	"IsBetaVersion": false,
+	"Installed": false,
+	"SupportedTargetPlatforms": [ "Win64" ],
+	"Modules": [
+		{
+			"Name": "MJPEGDecoderElectra",
+			"Type": "Runtime",
+			"LoadingPhase": "PostEngineInit",
+			"PlatformAllowList": [ "Win64" ],
+			"TargetDenyList": [ "Server" ]
+		}
+	],
+	"Plugins": [
+		{
+			"Name": "ElectraCodecs",
+			"Enabled": true
+		}
+	]
+}
diff --git a/Engine/Plugins/Media/MJPEGDecoderElectra/Source/MJPEGDecoderElectra/MJPEGDecoderElectra.Build.cs b/Engine/Plugins/Media/MJPEGDecoderElectra/Source/MJPEGDecoderElectra/MJPEGDecoderElectra.Build.cs
new file mode 100644
--- /dev/null
+++ b/Engine/Plugins/Media/MJPEGDecoderElectra/Source/MJPEGDecoderElectra/MJPEGDecoderElectra.Build.cs
@@ -0,0 +1,35 @@
+// Copyright Epic Games, Inc. All Rights Reserved.
+
+using System.IO;
+
+namespace UnrealBuildTool.Rules
+{
+	public class MJPEGDecoderElectra : ModuleRules
+	{
+		public MJPEGDecoderElectra(ReadOnlyTargetRules Target) : base(Target)
+		{
+			PrivateDependencyModuleNames.AddRange(
+				new string[] {
+					"Core",
+                    "Projects",
+                    "ElectraBase",
+                    "ElectraSamples",
+					"ElectraCodecFactory",
+					"ElectraDecoders"
+                });
+
+            if (Target.Platform == UnrealTargetPlatform.Win64)// || Target.Platform == UnrealTargetPlatform.Mac || Target.Platform == UnrealTargetPlatform.Linux)
+            {
+                DynamicallyLoadedModuleNames.AddRange(
+                    new string[] {
+                    "ImageWrapper",
+                    });
+
+                PrivateIncludePathModuleNames.AddRange(
+                    new string[] {
+                    "ImageWrapper",
+                    });
+			}
+		}
+	}
+}
diff --git a/Engine/Plugins/Media/MJPEGDecoderElectra/Source/MJPEGDecoderElectra/Private/MJPEGDecoder/ElectraMediaMJPEGDecoder.cpp b/Engine/Plugins/Media/MJPEGDecoderElectra/Source/MJPEGDecoderElectra/Private/MJPEGDecoder/ElectraMediaMJPEGDecoder.cpp
new file mode 100644
--- /dev/null
+++ b/Engine/Plugins/Media/MJPEGDecoderElectra/Source/MJPEGDecoderElectra/Private/MJPEGDecoder/ElectraMediaMJPEGDecoder.cpp
@@ -0,0 +1,538 @@
+// Copyright Epic Games, Inc. All Rights Reserved.
+
+#include "ElectraMediaMJPEGDecoder.h"
+#include "MJPEGDecoderElectraModule.h"
+#include "IElectraCodecRegistry.h"
+#include "IElectraCodecFactory.h"
+#include "Misc/ScopeLock.h"
+#include "Templates/SharedPointer.h"
+#include "Features/IModularFeatures.h"
+#include "Features/IModularFeature.h"
+#include "IElectraCodecFactoryModule.h"
+#include "IElectraDecoderFeaturesAndOptions.h"
+#include "IElectraDecoder.h"
+#include "IElectraDecoderOutputVideo.h"
+#include "ElectraDecodersUtils.h"
+#include "Utils/ElectraBitstreamProcessorDefault.h"
+
+#include "ElectraDecodersPlatformResources.h"
+
+#include "IImageWrapperModule.h"
+
+#define ERRCODE_INTERNAL_NO_ERROR							0
+#define ERRCODE_INTERNAL_ALREADY_CLOSED						1
+#define ERRCODE_INTERNAL_FAILED_TO_DECODE_INPUT				2
+
+/*********************************************************************************************************************/
+/*********************************************************************************************************************/
+/*********************************************************************************************************************/
+
+class FVideoDecoderMJPEGElectra;
+
+
+class FVideoDecoderOutputMJPEGElectra : public IElectraDecoderVideoOutput, public IElectraDecoderVideoOutputImageBuffers
+{
+public:
+	virtual ~FVideoDecoderOutputMJPEGElectra()
+	{
+	}
+
+	FTimespan GetPTS() const override
+	{
+		return PTS;
+	}
+	uint64 GetUserValue() const override
+	{
+		return UserValue;
+	}
+
+	EOutputType GetOutputType() const
+	{
+		return EOutputType::Output;
+	}
+
+	int32 GetWidth() const override
+	{
+		return Width - Crop.Left - Crop.Right;
+	}
+	int32 GetHeight() const override
+	{
+		return Height - Crop.Top - Crop.Bottom;
+	}
+	int32 GetDecodedWidth() const override
+	{
+		return Width;
+	}
+	int32 GetDecodedHeight() const override
+	{
+		return Height;
+	}
+	FElectraVideoDecoderOutputCropValues GetCropValues() const override
+	{
+		return Crop;
+	}
+	int32 GetAspectRatioW() const override
+	{
+		return AspectW;
+	}
+	int32 GetAspectRatioH() const override
+	{
+		return AspectH;
+	}
+	int32 GetFrameRateNumerator() const override
+	{
+		return FrameRateN;
+	}
+	int32 GetFrameRateDenominator() const override
+	{
+		return FrameRateD;
+	}
+	int32 GetNumberOfBits() const override
+	{
+		return NumBits;
+	}
+	void GetExtraValues(TMap<FString, FVariant>& OutExtraValues) const override
+	{
+		OutExtraValues = ExtraValues;
+	}
+	void* GetPlatformOutputHandle(EElectraDecoderPlatformOutputHandleType InTypeOfHandle) const override
+	{
+		if (InTypeOfHandle == EElectraDecoderPlatformOutputHandleType::ImageBuffers)
+		{
+			return static_cast<IElectraDecoderVideoOutputImageBuffers*>(const_cast<FVideoDecoderOutputMJPEGElectra*>(this));
+		}
+		return nullptr;
+	}
+	IElectraDecoderVideoOutput::EImageCopyResult CopyPlatformImage(IElectraDecoderVideoOutputCopyResources* InCopyResources) const override
+	{
+		return IElectraDecoderVideoOutput::EImageCopyResult::NotSupported;
+	}
+
+	// Methods from IElectraDecoderVideoOutputImageBuffers
+	uint32 GetCodec4CC() const override
+	{
+		return Codec4CC;
+	}
+	int32 GetNumberOfBuffers() const override
+	{
+		return NumBuffers;
+	}
+	TSharedPtr<TArray64<uint8>, ESPMode::ThreadSafe> GetBufferDataByIndex(int32 InBufferIndex) const override
+	{
+		if (InBufferIndex == 0)
+		{
+			return Buffers[EBufferType::Buffer_Color].Buffer;
+		}
+		else if (InBufferIndex == 1)
+		{
+			return Buffers[EBufferType::Buffer_Alpha].Buffer;
+		}
+		return nullptr;
+	}
+	void* GetBufferTextureByIndex(int32 InBufferIndex) const override
+	{
+		return nullptr;
+	}
+	bool GetBufferTextureSyncByIndex(int32 InBufferIndex, FElectraDecoderOutputSync& SyncObject) const override
+	{
+		return false;
+	}
+
+	EPixelFormat GetBufferFormatByIndex(int32 InBufferIndex) const override
+	{
+		if (InBufferIndex == 0)
+		{
+			return Buffers[EBufferType::Buffer_Color].BufferFormat;
+		}
+		else if (InBufferIndex == 1)
+		{
+			return Buffers[EBufferType::Buffer_Alpha].BufferFormat;
+		}
+		return EPixelFormat::PF_Unknown;
+	}
+	EElectraTextureSamplePixelEncoding GetBufferEncodingByIndex(int32 InBufferIndex) const override
+	{
+		if (InBufferIndex == 0)
+		{
+			return Buffers[EBufferType::Buffer_Color].BufferEncoding;
+		}
+		else if (InBufferIndex == 1)
+		{
+			return Buffers[EBufferType::Buffer_Alpha].BufferEncoding;
+		}
+		return EElectraTextureSamplePixelEncoding::Native;
+	}
+	int32 GetBufferPitchByIndex(int32 InBufferIndex) const override
+	{
+		if (InBufferIndex == 0)
+		{
+			return Buffers[EBufferType::Buffer_Color].Pitch;
+		}
+		else if (InBufferIndex == 1)
+		{
+			return Buffers[EBufferType::Buffer_Alpha].Pitch;
+		}
+		return 0;
+	}
+
+public:
+	FTimespan PTS;
+	uint64 UserValue = 0;
+
+	FElectraVideoDecoderOutputCropValues Crop;
+	int32 Width = 0;
+	int32 Height = 0;
+	int32 NumBits = 0;
+	int32 AspectW = 1;
+	int32 AspectH = 1;
+	int32 FrameRateN = 0;
+	int32 FrameRateD = 0;
+	TMap<FString, FVariant> ExtraValues;
+
+	uint32 Codec4CC = 0;
+	int32 NumBuffers = 0;
+
+	enum EBufferType
+	{
+		Buffer_Color = 0,
+		Buffer_Alpha = 1
+	};
+
+	struct FBufferInfo
+	{
+		TSharedPtr<TArray64<uint8>, ESPMode::ThreadSafe> Buffer;
+		EPixelFormat BufferFormat;
+		EElectraTextureSamplePixelEncoding BufferEncoding;
+		int32 Pitch = 0;
+	};
+	FBufferInfo Buffers[2];
+};
+
+
+class FVideoDecoderMJPEGElectra : public IElectraDecoder
+{
+public:
+	static void GetConfigurationOptions(TMap<FString, FVariant>& OutOptions)
+	{
+		OutOptions.Emplace(IElectraDecoderFeature::MinimumNumberOfOutputFrames, FVariant((int32)5));
+		OutOptions.Emplace(IElectraDecoderFeature::IsAdaptive, FVariant(true));
+	}
+
+	FVideoDecoderMJPEGElectra(const TMap<FString, FVariant>& InOptions);
+
+	virtual ~FVideoDecoderMJPEGElectra();
+
+	IElectraDecoder::EType GetType() const override
+	{
+		return IElectraDecoder::EType::Video;
+	}
+
+	void GetFeatures(TMap<FString, FVariant>& OutFeatures) const override;
+
+	FError GetError() const override;
+
+	void Close() override;
+	IElectraDecoder::ECSDCompatibility IsCompatibleWith(const TMap<FString, FVariant>& CSDAndAdditionalOptions) override;
+	bool ResetToCleanStart() override;
+
+	EDecoderError DecodeAccessUnit(const FInputAccessUnit& InInputAccessUnit, const TMap<FString, FVariant>& InAdditionalOptions) override;
+	EDecoderError SendEndOfData() override;
+	EDecoderError Flush() override;
+	EOutputStatus HaveOutput() override;
+	TSharedPtr<IElectraDecoderOutput, ESPMode::ThreadSafe> GetOutput() override;
+
+	TSharedPtr<IElectraDecoderBitstreamProcessor, ESPMode::ThreadSafe> CreateBitstreamProcessor() override
+	{ return FElectraDecoderBitstreamProcessorDefault::Create(); }
+
+	void Suspend() override
+	{ }
+	void Resume() override
+	{ }
+
+private:
+	static constexpr uint32 Make4CC(const uint8 A, const uint8 B, const uint8 C, const uint8 D)
+	{
+		return (static_cast<uint32>(A) << 24) | (static_cast<uint32>(B) << 16) | (static_cast<uint32>(C) << 8) | static_cast<uint32>(D);
+	}
+
+	bool PostError(int32 ApiReturnValue, FString Message, int32 Code);
+
+
+	IImageWrapperModule& ImageWrapperModule;
+
+	int32 DisplayWidth = 0;
+	int32 DisplayHeight = 0;
+	int32 AspectW = 0;
+	int32 AspectH = 0;
+	uint32 Codec4CC = 0;
+
+	IElectraDecoder::FError LastError;
+
+	TSharedPtr<FVideoDecoderOutputMJPEGElectra, ESPMode::ThreadSafe> CurrentOutput;
+	bool bFlushPending = false;
+};
+
+
+/*********************************************************************************************************************/
+/*********************************************************************************************************************/
+/*********************************************************************************************************************/
+
+class FMJPEGVideoDecoderElectraFactory : public IElectraCodecFactory, public IElectraCodecModularFeature, public TSharedFromThis<FMJPEGVideoDecoderElectraFactory, ESPMode::ThreadSafe>
+{
+public:
+	static constexpr uint32 Make4CC(const uint8 A, const uint8 B, const uint8 C, const uint8 D)
+	{
+		return (static_cast<uint32>(A) << 24) | (static_cast<uint32>(B) << 16) | (static_cast<uint32>(C) << 8) | static_cast<uint32>(D);
+	}
+
+	virtual ~FMJPEGVideoDecoderElectraFactory()
+	{ }
+
+	void GetListOfFactories(TArray<TWeakPtr<IElectraCodecFactory, ESPMode::ThreadSafe>>& OutCodecFactories) override
+	{
+		OutCodecFactories.Add(AsShared());
+	}
+
+	int32 SupportsFormat(TMap<FString, FVariant>& OutFormatInfo, const FString& InCodecFormat, bool bInEncoder, const TMap<FString, FVariant>& InOptions) const override
+	{
+		if (bInEncoder || !Permitted4CCs.Contains(InCodecFormat))
+		{
+			return 0;
+		}
+		switch((uint32)ElectraDecodersUtil::GetVariantValueSafeU64(InOptions, TEXT("codec_4cc"), 0))
+		{
+			case Make4CC('j','p','e','g'):
+			{
+				OutFormatInfo.Emplace(IElectraDecoderFormatInfo::HumanReadableFormatName, FVariant(FString(TEXT("MotionJPEG"))));
+				break;
+			}
+			default:
+			{
+				break;
+			}
+		}
+		OutFormatInfo.Emplace(IElectraDecoderFormatInfo::IsEveryFrameKeyframe, FVariant(true));
+		return 1;
+	}
+
+	void GetConfigurationOptions(TMap<FString, FVariant>& OutOptions) const override
+	{
+		FVideoDecoderMJPEGElectra::GetConfigurationOptions(OutOptions);
+	}
+
+	TSharedPtr<IElectraDecoder, ESPMode::ThreadSafe> CreateDecoderForFormat(const FString& InCodecFormat, const TMap<FString, FVariant>& InOptions) override
+	{
+		return MakeShared<FVideoDecoderMJPEGElectra, ESPMode::ThreadSafe>(InOptions);
+	}
+
+	static TSharedPtr<FMJPEGVideoDecoderElectraFactory, ESPMode::ThreadSafe> Self;
+	static TArray<FString> Permitted4CCs;
+};
+TSharedPtr<FMJPEGVideoDecoderElectraFactory, ESPMode::ThreadSafe> FMJPEGVideoDecoderElectraFactory::Self;
+TArray<FString> FMJPEGVideoDecoderElectraFactory::Permitted4CCs = { TEXT("jpeg") };
+
+/*********************************************************************************************************************/
+/*********************************************************************************************************************/
+/*********************************************************************************************************************/
+
+void FElectraMediaMJPEGDecoder::Startup()
+{
+	// Make sure the codec factory module has been loaded.
+	FModuleManager::Get().LoadModule(TEXT("ElectraCodecFactory"));
+
+	// Create an instance of the factory, which is also the modular feature.
+	check(!FMJPEGVideoDecoderElectraFactory::Self.IsValid());
+	FMJPEGVideoDecoderElectraFactory::Self = MakeShared<FMJPEGVideoDecoderElectraFactory, ESPMode::ThreadSafe>();
+	// Register as modular feature.
+	IModularFeatures::Get().RegisterModularFeature(IElectraCodecFactoryModule::GetModularFeatureName(), FMJPEGVideoDecoderElectraFactory::Self.Get());
+}
+
+void FElectraMediaMJPEGDecoder::Shutdown()
+{
+	IModularFeatures::Get().UnregisterModularFeature(IElectraCodecFactoryModule::GetModularFeatureName(), FMJPEGVideoDecoderElectraFactory::Self.Get());
+	FMJPEGVideoDecoderElectraFactory::Self.Reset();
+}
+
+/*********************************************************************************************************************/
+/*********************************************************************************************************************/
+/*********************************************************************************************************************/
+
+FVideoDecoderMJPEGElectra::FVideoDecoderMJPEGElectra(const TMap<FString, FVariant>& InOptions)
+	: ImageWrapperModule(FModuleManager::LoadModuleChecked<IImageWrapperModule>("ImageWrapper"))
+{
+	DisplayWidth = (int32)ElectraDecodersUtil::GetVariantValueSafeI64(InOptions, TEXT("width"), 0);
+	DisplayHeight = (int32)ElectraDecodersUtil::GetVariantValueSafeI64(InOptions, TEXT("height"), 0);
+	AspectW = (int32)ElectraDecodersUtil::GetVariantValueSafeI64(InOptions, TEXT("aspect_w"), 0);
+	AspectH = (int32)ElectraDecodersUtil::GetVariantValueSafeI64(InOptions, TEXT("aspect_h"), 0);
+	Codec4CC = (uint32)ElectraDecodersUtil::GetVariantValueSafeU64(InOptions, TEXT("codec_4cc"), 0);
+}
+
+FVideoDecoderMJPEGElectra::~FVideoDecoderMJPEGElectra()
+{
+	// Close() must have been called already!
+	check(LastError.Code == ERRCODE_INTERNAL_ALREADY_CLOSED);
+	// We do it nonetheless...
+	Close();
+}
+
+void FVideoDecoderMJPEGElectra::GetFeatures(TMap<FString, FVariant>& OutFeatures) const
+{
+	GetConfigurationOptions(OutFeatures);
+}
+
+IElectraDecoder::FError FVideoDecoderMJPEGElectra::GetError() const
+{
+	return LastError;
+}
+
+bool FVideoDecoderMJPEGElectra::PostError(int32 ApiReturnValue, FString Message, int32 Code)
+{
+	LastError.Code = Code;
+	LastError.SdkCode = ApiReturnValue;
+	LastError.Message = MoveTemp(Message);
+	return false;
+}
+
+void FVideoDecoderMJPEGElectra::Close()
+{
+	ResetToCleanStart();
+	// Set the error state that all subsequent calls will fail.
+	PostError(0, TEXT("Already closed"), ERRCODE_INTERNAL_ALREADY_CLOSED);
+}
+
+IElectraDecoder::ECSDCompatibility FVideoDecoderMJPEGElectra::IsCompatibleWith(const TMap<FString, FVariant>& CSDAndAdditionalOptions)
+{
+	return IElectraDecoder::ECSDCompatibility::Compatible;
+}
+
+bool FVideoDecoderMJPEGElectra::ResetToCleanStart()
+{
+	bFlushPending = false;
+	CurrentOutput.Reset();
+	return true;
+}
+
+
+IElectraDecoder::EDecoderError FVideoDecoderMJPEGElectra::DecodeAccessUnit(const FInputAccessUnit& InInputAccessUnit, const TMap<FString, FVariant>& InAdditionalOptions)
+{
+	// If already in error do nothing!
+	if (LastError.IsSet())
+	{
+		return IElectraDecoder::EDecoderError::Error;
+	}
+
+	// Can not feed new input until draining has finished.
+	if (bFlushPending)
+	{
+		return IElectraDecoder::EDecoderError::EndOfData;
+	}
+
+	// If there is pending output it is very likely that decoding this access unit would also generate output.
+	// Since that would result in loss of the pending output we return now.
+	if (CurrentOutput.IsValid())
+	{
+		return IElectraDecoder::EDecoderError::NoBuffer;
+	}
+
+	// Decode data. This immediately produces a new output frame.
+	if (InInputAccessUnit.Data && InInputAccessUnit.DataSize)
+	{
+		TSharedPtr<IImageWrapper> ImageWrapper = ImageWrapperModule.CreateImageWrapper(EImageFormat::JPEG);
+		if (!ImageWrapper.IsValid() || !ImageWrapper->SetCompressed(InInputAccessUnit.Data, InInputAccessUnit.DataSize))
+		{
+			PostError(0, TEXT("JPEG decoding failed"), ERRCODE_INTERNAL_FAILED_TO_DECODE_INPUT);
+			return IElectraDecoder::EDecoderError::Error;
+		}
+
+		int32 DecodedWidth = (int32)ImageWrapper->GetWidth();
+		int32 DecodedHeight = (int32)ImageWrapper->GetHeight();
+
+		TSharedPtr<FVideoDecoderOutputMJPEGElectra, ESPMode::ThreadSafe> NewOutput = MakeShared<FVideoDecoderOutputMJPEGElectra>();
+		NewOutput->PTS = InInputAccessUnit.PTS;
+		NewOutput->UserValue = InInputAccessUnit.UserValue;
+
+		NewOutput->Width = DecodedWidth;
+		NewOutput->Height = DecodedHeight;
+		NewOutput->Crop.Right = DecodedWidth - DisplayWidth;
+		NewOutput->Crop.Bottom = DecodedHeight - DisplayHeight;
+		if (AspectW && AspectH)
+		{
+			NewOutput->AspectW = AspectW;
+			NewOutput->AspectH = AspectH;
+		}
+
+		NewOutput->NumBits = 8;
+
+		NewOutput->NumBuffers = 1;
+		NewOutput->Buffers[0].Buffer = MakeShared<TArray64<uint8>, ESPMode::ThreadSafe>();
+		NewOutput->Buffers[0].BufferFormat = EPixelFormat::PF_B8G8R8A8;
+		NewOutput->Buffers[0].BufferEncoding = EElectraTextureSamplePixelEncoding::Native;
+		NewOutput->Buffers[0].Pitch = DecodedWidth * 4;
+		if (!ImageWrapper->GetRaw(ERGBFormat::BGRA, 8, *NewOutput->Buffers[0].Buffer))
+		{
+			PostError(0, TEXT("Failed to get decoded JPEG image data"), ERRCODE_INTERNAL_FAILED_TO_DECODE_INPUT);
+			return IElectraDecoder::EDecoderError::Error;
+		}
+
+		NewOutput->ExtraValues.Emplace(TEXT("codec"), FVariant(TEXT("mjpg")));
+		NewOutput->ExtraValues.Emplace(TEXT("codec_4cc"), FVariant(Codec4CC));
+
+		CurrentOutput = MoveTemp(NewOutput);
+	}
+	return IElectraDecoder::EDecoderError::None;
+}
+
+IElectraDecoder::EDecoderError FVideoDecoderMJPEGElectra::SendEndOfData()
+{
+	// If already in error do nothing!
+	if (LastError.IsSet())
+	{
+		return IElectraDecoder::EDecoderError::Error;
+	}
+	// Already draining?
+	if (bFlushPending)
+	{
+		return IElectraDecoder::EDecoderError::EndOfData;
+	}
+	bFlushPending = true;
+	return IElectraDecoder::EDecoderError::None;
+}
+
+IElectraDecoder::EDecoderError FVideoDecoderMJPEGElectra::Flush()
+{
+	// If already in error do nothing!
+	if (LastError.IsSet())
+	{
+		return IElectraDecoder::EDecoderError::Error;
+	}
+	ResetToCleanStart();
+	return IElectraDecoder::EDecoderError::None;
+}
+
+IElectraDecoder::EOutputStatus FVideoDecoderMJPEGElectra::HaveOutput()
+{
+	// If already in error do nothing!
+	if (LastError.IsSet())
+	{
+		return IElectraDecoder::EOutputStatus::Error;
+	}
+	// Have output?
+	if (CurrentOutput.IsValid())
+	{
+		return IElectraDecoder::EOutputStatus::Available;
+	}
+	// Pending flush?
+	if (bFlushPending)
+	{
+		bFlushPending = false;
+		return IElectraDecoder::EOutputStatus::EndOfData;
+	}
+	return IElectraDecoder::EOutputStatus::NeedInput;
+}
+
+TSharedPtr<IElectraDecoderOutput, ESPMode::ThreadSafe> FVideoDecoderMJPEGElectra::GetOutput()
+{
+	TSharedPtr<IElectraDecoderOutput, ESPMode::ThreadSafe> Out = CurrentOutput;
+	CurrentOutput.Reset();
+	return Out;
+}
diff --git a/Engine/Plugins/Media/MJPEGDecoderElectra/Source/MJPEGDecoderElectra/Private/MJPEGDecoder/ElectraMediaMJPEGDecoder.h b/Engine/Plugins/Media/MJPEGDecoderElectra/Source/MJPEGDecoderElectra/Private/MJPEGDecoder/ElectraMediaMJPEGDecoder.h
new file mode 100644
--- /dev/null
+++ b/Engine/Plugins/Media/MJPEGDecoderElectra/Source/MJPEGDecoderElectra/Private/MJPEGDecoder/ElectraMediaMJPEGDecoder.h
@@ -0,0 +1,11 @@
+// Copyright Epic Games, Inc. All Rights Reserved.
+#pragma once
+
+#include "CoreMinimal.h"
+
+class FElectraMediaMJPEGDecoder
+{
+public:
+	static void Startup();
+	static void Shutdown();
+};
diff --git a/Engine/Plugins/Media/MJPEGDecoderElectra/Source/MJPEGDecoderElectra/Private/MJPEGDecoderElectraModule.cpp b/Engine/Plugins/Media/MJPEGDecoderElectra/Source/MJPEGDecoderElectra/Private/MJPEGDecoderElectraModule.cpp
new file mode 100644
--- /dev/null
+++ b/Engine/Plugins/Media/MJPEGDecoderElectra/Source/MJPEGDecoderElectra/Private/MJPEGDecoderElectraModule.cpp
@@ -0,0 +1,33 @@
+// Copyright Epic Games, Inc. All Rights Reserved.
+
+#include "MJPEGDecoderElectraModule.h"
+#include "MJPEGDecoder/ElectraMediaMJPEGDecoder.h"
+#include "Modules/ModuleManager.h"
+
+#define LOCTEXT_NAMESPACE "MJPEGDecoderElectraModule"
+
+DEFINE_LOG_CATEGORY(LogMJPEGElectraDecoder);
+
+class FMJPEGElectraDecoderModule : public IModuleInterface
+{
+public:
+	void StartupModule() override
+	{
+		FElectraMediaMJPEGDecoder::Startup();
+	}
+
+	void ShutdownModule() override
+	{
+		FElectraMediaMJPEGDecoder::Shutdown();
+	}
+
+	bool SupportsDynamicReloading() override
+	{
+		// Codec could still be in use
+		return false;
+	}
+};
+
+IMPLEMENT_MODULE(FMJPEGElectraDecoderModule, MJPEGDecoderElectra);
+
+#undef LOCTEXT_NAMESPACE
diff --git a/Engine/Plugins/Media/MJPEGDecoderElectra/Source/MJPEGDecoderElectra/Private/MJPEGDecoderElectraModule.h b/Engine/Plugins/Media/MJPEGDecoderElectra/Source/MJPEGDecoderElectra/Private/MJPEGDecoderElectraModule.h
new file mode 100644
--- /dev/null
+++ b/Engine/Plugins/Media/MJPEGDecoderElectra/Source/MJPEGDecoderElectra/Private/MJPEGDecoderElectraModule.h
@@ -0,0 +1,8 @@
+// Copyright Epic Games, Inc. All Rights Reserved.
+
+#pragma once
+
+#include <CoreMinimal.h>
+#include "Logging/LogMacros.h"
+
+DECLARE_LOG_CATEGORY_EXTERN(LogMJPEGElectraDecoder, Log, All);
