ElectraCodecs: change mjpeg decoder to use ImageUtils instead of ImageWrapper.

#rnx

#Synced-CL Unknown

diff --git a/Engine/Plugins/Media/MJPEGDecoderElectra/MJPEGDecoderElectra.uplugin b/Engine/Plugins/Media/MJPEGDecoderElectra/MJPEGDecoderElectra.uplugin
--- a/Engine/Plugins/Media/MJPEGDecoderElectra/MJPEGDecoderElectra.uplugin
+++ b/Engine/Plugins/Media/MJPEGDecoderElectra/MJPEGDecoderElectra.uplugin
@@ -17,7 +17,7 @@
 			"Name": "MJPEGDecoderElectra",
 			"Type": "Runtime",
 			"LoadingPhase": "PostEngineInit",
-			"PlatformAllowList": [ "Win64" ],
+			"PlatformAllowList": [ "Win64", "Mac", "Linux" ],
 			"TargetDenyList": [ "Server" ]
 		}
 	],
diff --git a/Engine/Plugins/Media/MJPEGDecoderElectra/Source/MJPEGDecoderElectra/MJPEGDecoderElectra.Build.cs b/Engine/Plugins/Media/MJPEGDecoderElectra/Source/MJPEGDecoderElectra/MJPEGDecoderElectra.Build.cs
--- a/Engine/Plugins/Media/MJPEGDecoderElectra/Source/MJPEGDecoderElectra/MJPEGDecoderElectra.Build.cs
+++ b/Engine/Plugins/Media/MJPEGDecoderElectra/Source/MJPEGDecoderElectra/MJPEGDecoderElectra.Build.cs
@@ -11,25 +11,13 @@ namespace UnrealBuildTool.Rules
 			PrivateDependencyModuleNames.AddRange(
 				new string[] {
 					"Core",
+					"Engine",
                     "Projects",
                     "ElectraBase",
                     "ElectraSamples",
 					"ElectraCodecFactory",
 					"ElectraDecoders"
                 });
-
-            if (Target.Platform == UnrealTargetPlatform.Win64)// || Target.Platform == UnrealTargetPlatform.Mac || Target.Platform == UnrealTargetPlatform.Linux)
-            {
-                DynamicallyLoadedModuleNames.AddRange(
-                    new string[] {
-                    "ImageWrapper",
-                    });
-
-                PrivateIncludePathModuleNames.AddRange(
-                    new string[] {
-                    "ImageWrapper",
-                    });
-			}
 		}
 	}
 }
diff --git a/Engine/Plugins/Media/MJPEGDecoderElectra/Source/MJPEGDecoderElectra/Private/MJPEGDecoder/ElectraMediaMJPEGDecoder.cpp b/Engine/Plugins/Media/MJPEGDecoderElectra/Source/MJPEGDecoderElectra/Private/MJPEGDecoder/ElectraMediaMJPEGDecoder.cpp
--- a/Engine/Plugins/Media/MJPEGDecoderElectra/Source/MJPEGDecoderElectra/Private/MJPEGDecoder/ElectraMediaMJPEGDecoder.cpp
+++ b/Engine/Plugins/Media/MJPEGDecoderElectra/Source/MJPEGDecoderElectra/Private/MJPEGDecoder/ElectraMediaMJPEGDecoder.cpp
@@ -17,7 +17,7 @@
 
 #include "ElectraDecodersPlatformResources.h"
 
-#include "IImageWrapperModule.h"
+#include "ImageUtils.h"
 
 #define ERRCODE_INTERNAL_NO_ERROR							0
 #define ERRCODE_INTERNAL_ALREADY_CLOSED						1
@@ -253,8 +253,6 @@ private:
 	bool PostError(int32 ApiReturnValue, FString Message, int32 Code);
 
 
-	IImageWrapperModule& ImageWrapperModule;
-
 	int32 DisplayWidth = 0;
 	int32 DisplayHeight = 0;
 	int32 AspectW = 0;
@@ -318,7 +316,7 @@ public:
 	const IProviderInformation& GetProviderInformation() const override
 	{ return *this; }
 	FString GetName() const override
-	{ return FString(TEXT("ImageWrapper")); }
+	{ return FString(TEXT("ImageUtils")); }
 	FString GetVersion() const override
 	{ return(FString()); }
 	FString GetImplementation() const override
@@ -359,7 +357,6 @@ void FElectraMediaMJPEGDecoder::Shutdown()
 /*********************************************************************************************************************/
 
 FVideoDecoderMJPEGElectra::FVideoDecoderMJPEGElectra(const TMap<FString, FVariant>& InOptions)
-	: ImageWrapperModule(FModuleManager::LoadModuleChecked<IImageWrapperModule>("ImageWrapper"))
 {
 	DisplayWidth = (int32)ElectraDecodersUtil::GetVariantValueSafeI64(InOptions, TEXT("width"), 0);
 	DisplayHeight = (int32)ElectraDecodersUtil::GetVariantValueSafeI64(InOptions, TEXT("height"), 0);
@@ -438,15 +435,20 @@ IElectraDecoder::EDecoderError FVideoDecoderMJPEGElectra::DecodeAccessUnit(const
 	// Decode data. This immediately produces a new output frame.
 	if (InInputAccessUnit.Data && InInputAccessUnit.DataSize)
 	{
-		TSharedPtr<IImageWrapper> ImageWrapper = ImageWrapperModule.CreateImageWrapper(EImageFormat::JPEG);
-		if (!ImageWrapper.IsValid() || !ImageWrapper->SetCompressed(InInputAccessUnit.Data, InInputAccessUnit.DataSize))
+		FImage Image;
+		if (!FImageUtils::DecompressImage(InInputAccessUnit.Data, InInputAccessUnit.DataSize, Image))
 		{
 			PostError(0, TEXT("JPEG decoding failed"), ERRCODE_INTERNAL_FAILED_TO_DECODE_INPUT);
 			return IElectraDecoder::EDecoderError::Error;
 		}
+		if (Image.Format != ERawImageFormat::BGRA8)
+		{
+			PostError(0, TEXT("Decoded JPEG image is not BGRA"), ERRCODE_INTERNAL_FAILED_TO_DECODE_INPUT);
+			return IElectraDecoder::EDecoderError::Error;
+		}
 
-		int32 DecodedWidth = (int32)ImageWrapper->GetWidth();
-		int32 DecodedHeight = (int32)ImageWrapper->GetHeight();
+		int32 DecodedWidth = Image.SizeX;
+		int32 DecodedHeight = Image.SizeY;
 
 		TSharedPtr<FVideoDecoderOutputMJPEGElectra, ESPMode::ThreadSafe> NewOutput = MakeShared<FVideoDecoderOutputMJPEGElectra>();
 		NewOutput->PTS = InInputAccessUnit.PTS;
@@ -454,8 +456,8 @@ IElectraDecoder::EDecoderError FVideoDecoderMJPEGElectra::DecodeAccessUnit(const
 
 		NewOutput->Width = DecodedWidth;
 		NewOutput->Height = DecodedHeight;
-		NewOutput->Crop.Right = DecodedWidth - DisplayWidth;
-		NewOutput->Crop.Bottom = DecodedHeight - DisplayHeight;
+		NewOutput->Crop.Right = DecodedWidth > DisplayWidth ? DecodedWidth - DisplayWidth : 0;
+		NewOutput->Crop.Bottom = DecodedHeight > DisplayHeight ? DecodedHeight - DisplayHeight : 0;
 		if (AspectW && AspectH)
 		{
 			NewOutput->AspectW = AspectW;
@@ -465,16 +467,13 @@ IElectraDecoder::EDecoderError FVideoDecoderMJPEGElectra::DecodeAccessUnit(const
 		NewOutput->NumBits = 8;
 
 		NewOutput->NumBuffers = 1;
-		NewOutput->Buffers[0].Buffer = MakeShared<TArray64<uint8>, ESPMode::ThreadSafe>();
+		TSharedPtr<TArray64<uint8>, ESPMode::ThreadSafe> Buf = MakeShareable(new TArray64<uint8>());
+		*Buf = MoveTemp(Image.RawData);
+		NewOutput->Buffers[0].Buffer = MoveTemp(Buf);
 		NewOutput->Buffers[0].BufferFormat = EPixelFormat::PF_B8G8R8A8;
 		NewOutput->Buffers[0].BufferEncoding = EElectraTextureSamplePixelEncoding::Native;
 		NewOutput->Buffers[0].Pitch = DecodedWidth * 4;
-		if (!ImageWrapper->GetRaw(ERGBFormat::BGRA, 8, *NewOutput->Buffers[0].Buffer))
-		{
-			PostError(0, TEXT("Failed to get decoded JPEG image data"), ERRCODE_INTERNAL_FAILED_TO_DECODE_INPUT);
-			return IElectraDecoder::EDecoderError::Error;
-		}
-
+		NewOutput->Codec4CC = Codec4CC;
 		NewOutput->ExtraValues.Emplace(TEXT("codec"), FVariant(TEXT("mjpg")));
 		NewOutput->ExtraValues.Emplace(TEXT("codec_4cc"), FVariant(Codec4CC));
 
